Build Storage Tool - Added shell script to start Build Storage Tool that is more easily accessable for when downloaded directly from Horde. Also added PublishToHorde option which were implied before (defaults to true as before)

diff --git a/Engine/Programs/BuildStorageTool/Build/Linux/RunBuildStorageTool.sh b/Engine/Programs/BuildStorageTool/Build/Linux/RunBuildStorageTool.sh
new file mode 100644
--- /dev/null
+++ b/Engine/Programs/BuildStorageTool/Build/Linux/RunBuildStorageTool.sh
@@ -0,0 +1,3 @@
+#!/bin/bash
+SCRIPT_DIR=$( cd -- "$( dirname -- "${BASH_SOURCE[0]}" )" &> /dev/null && pwd )
+"$SCRIPT_DIR"/Linux/Engine/Binaries/Linux/buildstoragetool "$@" 2>&1 &
\ No newline at end of file
diff --git a/Engine/Programs/BuildStorageTool/Build/Mac/RunBuildStorageTool.sh b/Engine/Programs/BuildStorageTool/Build/Mac/RunBuildStorageTool.sh
new file mode 100644
--- /dev/null
+++ b/Engine/Programs/BuildStorageTool/Build/Mac/RunBuildStorageTool.sh
@@ -0,0 +1,3 @@
+#!/bin/bash
+SCRIPT_DIR=$( cd -- "$( dirname -- "${BASH_SOURCE[0]}" )" &> /dev/null && pwd )
+"$SCRIPT_DIR"/Mac/Engine/Binaries/Mac/buildstoragetool "$@" 2>&1 &
\ No newline at end of file
diff --git a/Engine/Programs/BuildStorageTool/Build/Win64/RunBuildStorageTool.bat b/Engine/Programs/BuildStorageTool/Build/Win64/RunBuildStorageTool.bat
new file mode 100755
--- /dev/null
+++ b/Engine/Programs/BuildStorageTool/Build/Win64/RunBuildStorageTool.bat
@@ -0,0 +1 @@
+START %~dp0/Windows/Engine/Binaries/Win64/BuildStorageTool.exe %*
\ No newline at end of file
