Build Storage Tool - Fixed casing mismatch in RunBuildStorageTool.sh scripts for linux and mac

diff --git a/Engine/Programs/BuildStorageTool/Build/Linux/RunBuildStorageTool.sh b/Engine/Programs/BuildStorageTool/Build/Linux/RunBuildStorageTool.sh
--- a/Engine/Programs/BuildStorageTool/Build/Linux/RunBuildStorageTool.sh
+++ b/Engine/Programs/BuildStorageTool/Build/Linux/RunBuildStorageTool.sh
@@ -1,3 +1,3 @@
 #!/bin/bash
 SCRIPT_DIR=$( cd -- "$( dirname -- "${BASH_SOURCE[0]}" )" &> /dev/null && pwd )
-"$SCRIPT_DIR"/Linux/Engine/Binaries/Linux/buildstoragetool "$@" 2>&1 &
\ No newline at end of file
+"$SCRIPT_DIR"/Linux/Engine/Binaries/Linux/BuildStorageTool "$@" 2>&1 &
\ No newline at end of file
diff --git a/Engine/Programs/BuildStorageTool/Build/Mac/RunBuildStorageTool.sh b/Engine/Programs/BuildStorageTool/Build/Mac/RunBuildStorageTool.sh
--- a/Engine/Programs/BuildStorageTool/Build/Mac/RunBuildStorageTool.sh
+++ b/Engine/Programs/BuildStorageTool/Build/Mac/RunBuildStorageTool.sh
@@ -1,3 +1,3 @@
 #!/bin/bash
 SCRIPT_DIR=$( cd -- "$( dirname -- "${BASH_SOURCE[0]}" )" &> /dev/null && pwd )
-"$SCRIPT_DIR"/Mac/Engine/Binaries/Mac/buildstoragetool "$@" 2>&1 &
\ No newline at end of file
+"$SCRIPT_DIR"/Mac/Engine/Binaries/Mac/BuildStorageTool "$@" 2>&1 &
\ No newline at end of file
