Build Storage Tool - Fixed issue with the run script on windows for paths that contain spaces

diff --git a/Engine/Programs/BuildStorageTool/Build/Win64/RunBuildStorageTool.bat b/Engine/Programs/BuildStorageTool/Build/Win64/RunBuildStorageTool.bat
--- a/Engine/Programs/BuildStorageTool/Build/Win64/RunBuildStorageTool.bat
+++ b/Engine/Programs/BuildStorageTool/Build/Win64/RunBuildStorageTool.bat
@@ -1 +1 @@
-START %~dp0/Windows/Engine/Binaries/Win64/BuildStorageTool.exe %*
\ No newline at end of file
+START "" "%~dp0Windows\Engine\Binaries\Win64\BuildStorageTool.exe" %*
\ No newline at end of file
