Build Storage Tool - Updated the run script on mac to work correctly

diff --git a/Engine/Programs/BuildStorageTool/Build/Mac/RunBuildStorageTool.sh b/Engine/Programs/BuildStorageTool/Build/Mac/RunBuildStorageTool.sh
--- a/Engine/Programs/BuildStorageTool/Build/Mac/RunBuildStorageTool.sh
+++ b/Engine/Programs/BuildStorageTool/Build/Mac/RunBuildStorageTool.sh
@@ -1,3 +1,3 @@
 #!/bin/bash
 SCRIPT_DIR=$( cd -- "$( dirname -- "${BASH_SOURCE[0]}" )" &> /dev/null && pwd )
-"$SCRIPT_DIR"/Mac/Engine/Binaries/Mac/BuildStorageTool "$@" 2>&1 &
\ No newline at end of file
+open "$SCRIPT_DIR"/Mac/BuildStorageTool.app --args "$@"
\ No newline at end of file
