BuildStorageTool: Set n.VerifyPeer to true in DefaultEngine.ini so that we stage/ship the public cert store when producing BuildStorageTool builds.

#Synced-CL 51039791

diff --git a/Engine/Programs/BuildStorageTool/Config/DefaultEngine.ini b/Engine/Programs/BuildStorageTool/Config/DefaultEngine.ini
--- a/Engine/Programs/BuildStorageTool/Config/DefaultEngine.ini
+++ b/Engine/Programs/BuildStorageTool/Config/DefaultEngine.ini
@@ -1,3 +1,7 @@
+
+[/Script/Engine.NetworkSettings]
+n.VerifyPeer=true
+
 [StudioTelemetry.Provider.EpicStudio]
 Enabled=True
 
