BuildStorageTool: Allow windows icon directory to be included in staging.

#Synced-CL 51529686

diff --git a/Engine/Programs/BuildStorageTool/Config/DefaultGame.ini b/Engine/Programs/BuildStorageTool/Config/DefaultGame.ini
--- a/Engine/Programs/BuildStorageTool/Config/DefaultGame.ini
+++ b/Engine/Programs/BuildStorageTool/Config/DefaultGame.ini
@@ -1,6 +1,7 @@
 [Staging]
 +AllowedDirectories=Engine/Content/Editor/Slate
 +AllowedDirectories=Engine/Content/Slate
++AllowedDirectories=Engine/Programs/BuildStorageTool/Resources/Windows
 +AllowedConfigFiles=Engine/Config/BaseEditorPerProjectUserSettings.ini
 +AllowedConfigFiles=BuildStorageTool/Config/DeployId.ini
 
